Updating Editor Mode plugin template to follow UE5 guidelines and use UEdMode instead of FEdMode

diff --git a/Engine/Plugins/Editor/PluginBrowser/Templates/EditorMode/Source/PLUGIN_NAME/PLUGIN_NAME.Build.cs b/Engine/Plugins/Editor/PluginBrowser/Templates/EditorMode/Source/PLUGIN_NAME/PLUGIN_NAME.Build.cs
--- a/Engine/Plugins/Editor/PluginBrowser/Templates/EditorMode/Source/PLUGIN_NAME/PLUGIN_NAME.Build.cs
+++ b/Engine/Plugins/Editor/PluginBrowser/Templates/EditorMode/Source/PLUGIN_NAME/PLUGIN_NAME.Build.cs
@@ -42,6 +42,7 @@ public class PLUGIN_NAME : ModuleRules
 				"EditorFramework",
 				"UnrealEd",
 				"LevelEditor",
+				"EditorStyle",
 				// ... add private dependencies that you statically link with here ...	
 			}
 			);
diff --git a/Engine/Plugins/Editor/PluginBrowser/Templates/EditorMode/Source/PLUGIN_NAME/Private/PLUGIN_NAME.cpp b/Engine/Plugins/Editor/PluginBrowser/Templates/EditorMode/Source/PLUGIN_NAME/Private/PLUGIN_NAME.cpp
--- a/Engine/Plugins/Editor/PluginBrowser/Templates/EditorMode/Source/PLUGIN_NAME/Private/PLUGIN_NAME.cpp
+++ b/Engine/Plugins/Editor/PluginBrowser/Templates/EditorMode/Source/PLUGIN_NAME/Private/PLUGIN_NAME.cpp
@@ -1,23 +1,17 @@
 // Copyright Epic Games, Inc. All Rights Reserved.
 
 #include "PLUGIN_NAME.h"
-#include "PLUGIN_NAMEEdMode.h"
+#include "Modules/ModuleInterface.h"
 
-#define LOCTEXT_NAMESPACE "FPLUGIN_NAMEModule"
-
-void FPLUGIN_NAMEModule::StartupModule()
+class FPLUGIN_NAMEModule : public IModuleInterface
 {
-	// This code will execute after your module is loaded into memory; the exact timing is specified in the .uplugin file per-module
-	FEditorModeRegistry::Get().RegisterMode<FPLUGIN_NAMEEdMode>(FPLUGIN_NAMEEdMode::EM_PLUGIN_NAMEEdModeId, LOCTEXT("PLUGIN_NAMEEdModeName", "PLUGIN_NAMEEdMode"), FSlateIcon(), true);
-}
+	/** This function will execute after your module is loaded into memory; the exact timing is specified in the .uplugin file per-module */
+	//virtual void StartupModule() override;
 
-void FPLUGIN_NAMEModule::ShutdownModule()
-{
-	// This function may be called during shutdown to clean up your module.  For modules that support dynamic reloading,
-	// we call this function before unloading the module.
-	FEditorModeRegistry::Get().UnregisterMode(FPLUGIN_NAMEEdMode::EM_PLUGIN_NAMEEdModeId);
-}
+	/** This function may be called during shutdown to clean up your module.  For modules that support dynamic reloading,
+	 *  we call this function before unloading the module
+	 */
+	//virtual void ShutdownModule() override;
+};
 
-#undef LOCTEXT_NAMESPACE
-	
 IMPLEMENT_MODULE(FPLUGIN_NAMEModule, PLUGIN_NAME)
\ No newline at end of file
diff --git a/Engine/Plugins/Editor/PluginBrowser/Templates/EditorMode/Source/PLUGIN_NAME/Private/PLUGIN_NAMEEdMode.cpp b/Engine/Plugins/Editor/PluginBrowser/Templates/EditorMode/Source/PLUGIN_NAME/Private/PLUGIN_NAMEEdMode.cpp
--- a/Engine/Plugins/Editor/PluginBrowser/Templates/EditorMode/Source/PLUGIN_NAME/Private/PLUGIN_NAMEEdMode.cpp
+++ b/Engine/Plugins/Editor/PluginBrowser/Templates/EditorMode/Source/PLUGIN_NAME/Private/PLUGIN_NAMEEdMode.cpp
@@ -5,21 +5,28 @@
 #include "Toolkits/ToolkitManager.h"
 #include "EditorModeManager.h"
 
-const FEditorModeID FPLUGIN_NAMEEdMode::EM_PLUGIN_NAMEEdModeId = TEXT("EM_PLUGIN_NAMEEdMode");
+#define LOCTEXT_NAMESPACE "PLUGIN_NAME"
 
-FPLUGIN_NAMEEdMode::FPLUGIN_NAMEEdMode()
-{
-
-}
+const FEditorModeID UPLUGIN_NAMEEdMode::EM_PLUGIN_NAMEEdModeId = TEXT("EM_PLUGIN_NAMEEdMode");
 
-FPLUGIN_NAMEEdMode::~FPLUGIN_NAMEEdMode()
+UPLUGIN_NAMEEdMode::UPLUGIN_NAMEEdMode()
+	: Super()
 {
-
+	FModuleManager::Get().LoadModule("EditorStyle");
+
+	Info = FEditorModeInfo(
+		FName(TEXT("PLUGIN_NAME")),
+		LOCTEXT("ModeName", "PLUGIN_NAME Editor Mode"),
+		FSlateIcon(FEditorStyle::GetStyleSetName(), "LevelEditor.MeshPaintMode", "LevelEditor.MeshPaintMode.Small"),
+		true,
+		600
+	);
 }
 
-void FPLUGIN_NAMEEdMode::Enter()
+
+void UPLUGIN_NAMEEdMode::Enter()
 {
-	FEdMode::Enter();
+	Super::Enter();
 
 	if (!Toolkit.IsValid() && UsesToolkits())
 	{
@@ -28,7 +35,7 @@ void FPLUGIN_NAMEEdMode::Enter()
 	}
 }
 
-void FPLUGIN_NAMEEdMode::Exit()
+void UPLUGIN_NAMEEdMode::Exit()
 {
 	if (Toolkit.IsValid())
 	{
@@ -37,14 +44,12 @@ void FPLUGIN_NAMEEdMode::Exit()
 	}
 
 	// Call base Exit method to ensure proper cleanup
-	FEdMode::Exit();
+	Super::Exit();
 }
 
-bool FPLUGIN_NAMEEdMode::UsesToolkits() const
+void UPLUGIN_NAMEEdMode::CreateToolkit()
 {
-	return true;
+	Toolkit = MakeShareable(new FPLUGIN_NAMEEdModeToolkit);
 }
 
-
-
-
+#undef LOCTEXT_NAMESPACE
diff --git a/Engine/Plugins/Editor/PluginBrowser/Templates/EditorMode/Source/PLUGIN_NAME/Private/PLUGIN_NAMEEdModeToolkit.cpp b/Engine/Plugins/Editor/PluginBrowser/Templates/EditorMode/Source/PLUGIN_NAME/Private/PLUGIN_NAMEEdModeToolkit.cpp
--- a/Engine/Plugins/Editor/PluginBrowser/Templates/EditorMode/Source/PLUGIN_NAME/Private/PLUGIN_NAMEEdModeToolkit.cpp
+++ b/Engine/Plugins/Editor/PluginBrowser/Templates/EditorMode/Source/PLUGIN_NAME/Private/PLUGIN_NAMEEdModeToolkit.cpp
@@ -13,7 +13,7 @@ FPLUGIN_NAMEEdModeToolkit::FPLUGIN_NAMEEdModeToolkit()
 {
 }
 
-void FPLUGIN_NAMEEdModeToolkit::Init(const TSharedPtr<IToolkitHost>& InitToolkitHost)
+void FPLUGIN_NAMEEdModeToolkit::Init(const TSharedPtr<IToolkitHost>& InitToolkitHost, TWeakObjectPtr<UEdMode> InOwningMode)
 {
 	struct Locals
 	{
@@ -102,8 +102,8 @@ void FPLUGIN_NAMEEdModeToolkit::Init(const TSharedPtr<IToolkitHost>& InitToolkit
 				]
 
 		];
-		
-	FModeToolkit::Init(InitToolkitHost);
+
+	FModeToolkit::Init(InitToolkitHost, InOwningMode);
 }
 
 FName FPLUGIN_NAMEEdModeToolkit::GetToolkitFName() const
@@ -116,9 +116,4 @@ FText FPLUGIN_NAMEEdModeToolkit::GetBaseToolkitName() const
 	return NSLOCTEXT("PLUGIN_NAMEEdModeToolkit", "DisplayName", "PLUGIN_NAMEEdMode Tool");
 }
 
-class FEdMode* FPLUGIN_NAMEEdModeToolkit::GetEditorMode() const
-{
-	return GLevelEditorModeTools().GetActiveMode(FPLUGIN_NAMEEdMode::EM_PLUGIN_NAMEEdModeId);
-}
-
 #undef LOCTEXT_NAMESPACE
diff --git a/Engine/Plugins/Editor/PluginBrowser/Templates/EditorMode/Source/PLUGIN_NAME/Public/PLUGIN_NAME.h b/Engine/Plugins/Editor/PluginBrowser/Templates/EditorMode/Source/PLUGIN_NAME/Public/PLUGIN_NAME.h
--- a/Engine/Plugins/Editor/PluginBrowser/Templates/EditorMode/Source/PLUGIN_NAME/Public/PLUGIN_NAME.h
+++ b/Engine/Plugins/Editor/PluginBrowser/Templates/EditorMode/Source/PLUGIN_NAME/Public/PLUGIN_NAME.h
@@ -1,15 +1,6 @@
 // Copyright Epic Games, Inc. All Rights Reserved.
 
 #pragma once
-
 #include "CoreMinimal.h"
-#include "Modules/ModuleManager.h"
-
-class FPLUGIN_NAMEModule : public IModuleInterface
-{
-public:
 
-	/** IModuleInterface implementation */
-	virtual void StartupModule() override;
-	virtual void ShutdownModule() override;
-};
+/** See FPLUGIN_NAMEModule definition in PLUGIN_NAME.cpp file */
diff --git a/Engine/Plugins/Editor/PluginBrowser/Templates/EditorMode/Source/PLUGIN_NAME/Public/PLUGIN_NAMEEdMode.h b/Engine/Plugins/Editor/PluginBrowser/Templates/EditorMode/Source/PLUGIN_NAME/Public/PLUGIN_NAMEEdMode.h
--- a/Engine/Plugins/Editor/PluginBrowser/Templates/EditorMode/Source/PLUGIN_NAME/Public/PLUGIN_NAMEEdMode.h
+++ b/Engine/Plugins/Editor/PluginBrowser/Templates/EditorMode/Source/PLUGIN_NAME/Public/PLUGIN_NAMEEdMode.h
@@ -2,23 +2,26 @@
 
 #pragma once
 
-#include "CoreMinimal.h"
-#include "EdMode.h"
+#include "Tools/UEdMode.h"
+#include "PLUGIN_NAMEEdMode.generated.h"
 
-class FPLUGIN_NAMEEdMode : public FEdMode
+
+UCLASS()
+class UPLUGIN_NAMEEdMode : public UEdMode
 {
 public:
+	GENERATED_BODY()
+
 	const static FEditorModeID EM_PLUGIN_NAMEEdModeId;
 public:
-	FPLUGIN_NAMEEdMode();
-	virtual ~FPLUGIN_NAMEEdMode();
+	UPLUGIN_NAMEEdMode();
 
-	// FEdMode interface
+	// UEdMode interface - see interface for more overrides
 	virtual void Enter() override;
 	virtual void Exit() override;
+	virtual void CreateToolkit() override;
 	//virtual void Tick(FEditorViewportClient* ViewportClient, float DeltaTime) override;
 	//virtual void Render(const FSceneView* View, FViewport* Viewport, FPrimitiveDrawInterface* PDI) override;
 	//virtual void ActorSelectionChangeNotify() override;
-	bool UsesToolkits() const override;
-	// End of FEdMode interface
+	// End of UEdMode interface
 };
diff --git a/Engine/Plugins/Editor/PluginBrowser/Templates/EditorMode/Source/PLUGIN_NAME/Public/PLUGIN_NAMEEdModeToolkit.h b/Engine/Plugins/Editor/PluginBrowser/Templates/EditorMode/Source/PLUGIN_NAME/Public/PLUGIN_NAMEEdModeToolkit.h
--- a/Engine/Plugins/Editor/PluginBrowser/Templates/EditorMode/Source/PLUGIN_NAME/Public/PLUGIN_NAMEEdModeToolkit.h
+++ b/Engine/Plugins/Editor/PluginBrowser/Templates/EditorMode/Source/PLUGIN_NAME/Public/PLUGIN_NAMEEdModeToolkit.h
@@ -12,12 +12,11 @@ public:
 	FPLUGIN_NAMEEdModeToolkit();
 	
 	/** FModeToolkit interface */
-	virtual void Init(const TSharedPtr<IToolkitHost>& InitToolkitHost) override;
+	virtual void Init(const TSharedPtr<IToolkitHost>& InitToolkitHost, TWeakObjectPtr<UEdMode> InOwningMode) override;
 
 	/** IToolkit interface */
 	virtual FName GetToolkitFName() const override;
 	virtual FText GetBaseToolkitName() const override;
-	virtual class FEdMode* GetEditorMode() const override;
 	virtual TSharedPtr<class SWidget> GetInlineContent() const override { return ToolkitWidget; }
 
 private:
