Updating Editor Mode plugin template to follow UE5 guidelines and use UEdMode instead of FEdMode

diff --git a/Engine/Plugins/Editor/PluginBrowser/Templates/EditorMode/Source/PLUGIN_NAME/PLUGIN_NAME.Build.cs b/Engine/Plugins/Editor/PluginBrowser/Templates/EditorMode/Source/PLUGIN_NAME/PLUGIN_NAME.Build.cs
--- a/Engine/Plugins/Editor/PluginBrowser/Templates/EditorMode/Source/PLUGIN_NAME/PLUGIN_NAME.Build.cs
+++ b/Engine/Plugins/Editor/PluginBrowser/Templates/EditorMode/Source/PLUGIN_NAME/PLUGIN_NAME.Build.cs
@@ -40,9 +40,11 @@ public class PLUGIN_NAME : ModuleRules
 				"SlateCore",
 				"InputCore",
 				"EditorFramework",
+				"EditorStyle",
 				"UnrealEd",
 				"LevelEditor",
-				"EditorStyle",
+				"InteractiveToolsFramework",
+				"EditorInteractiveToolsFramework"
 				// ... add private dependencies that you statically link with here ...	
 			}
 			);
diff --git a/Engine/Plugins/Editor/PluginBrowser/Templates/EditorMode/Source/PLUGIN_NAME/Private/PLUGIN_NAME.cpp b/Engine/Plugins/Editor/PluginBrowser/Templates/EditorMode/Source/PLUGIN_NAME/Private/PLUGIN_NAME.cpp
deleted file mode 100644
diff --git a/Engine/Plugins/Editor/PluginBrowser/Templates/EditorMode/Source/PLUGIN_NAME/Private/PLUGIN_NAMEEdMode.cpp b/Engine/Plugins/Editor/PluginBrowser/Templates/EditorMode/Source/PLUGIN_NAME/Private/PLUGIN_NAMEEdMode.cpp
deleted file mode 100644
diff --git a/Engine/Plugins/Editor/PluginBrowser/Templates/EditorMode/Source/PLUGIN_NAME/Private/PLUGIN_NAMEEdModeToolkit.cpp b/Engine/Plugins/Editor/PluginBrowser/Templates/EditorMode/Source/PLUGIN_NAME/Private/PLUGIN_NAMEEdModeToolkit.cpp
deleted file mode 100644
diff --git a/Engine/Plugins/Editor/PluginBrowser/Templates/EditorMode/Source/PLUGIN_NAME/Private/PLUGIN_NAMEEditorMode.cpp b/Engine/Plugins/Editor/PluginBrowser/Templates/EditorMode/Source/PLUGIN_NAME/Private/PLUGIN_NAMEEditorMode.cpp
new file mode 100644
--- /dev/null
+++ b/Engine/Plugins/Editor/PluginBrowser/Templates/EditorMode/Source/PLUGIN_NAME/Private/PLUGIN_NAMEEditorMode.cpp
@@ -0,0 +1,79 @@
+// Copyright Epic Games, Inc. All Rights Reserved.
+
+#include "PLUGIN_NAMEEditorMode.h"
+#include "PLUGIN_NAMEEditorModeToolkit.h"
+#include "EdModeInteractiveToolsContext.h"
+#include "InteractiveToolManager.h"
+#include "PLUGIN_NAMEEditorModeCommands.h"
+
+
+//////////////////////////////////////////////////////////////////////////
+////////////////////////////////////////////////////////////////////////// 
+// AddYourTool Step 1 - include the header file for your Tools here
+//////////////////////////////////////////////////////////////////////////
+////////////////////////////////////////////////////////////////////////// 
+#include "Tools/PLUGIN_NAMESimpleTool.h"
+#include "Tools/PLUGIN_NAMEInteractiveTool.h"
+
+// step 2: register a ToolBuilder in FPLUGIN_NAMEEditorMode::Enter() below
+
+
+#define LOCTEXT_NAMESPACE "PLUGIN_NAMEEditorMode"
+
+const FEditorModeID UPLUGIN_NAMEEditorMode::EM_PLUGIN_NAMEEditorModeId = TEXT("EM_PLUGIN_NAMEEditorMode");
+
+FString UPLUGIN_NAMEEditorMode::SimpleToolName = TEXT("PLUGIN_NAME_ActorInfoTool");
+FString UPLUGIN_NAMEEditorMode::InteractiveToolName = TEXT("PLUGIN_NAME_MeasureDistanceTool");
+
+
+UPLUGIN_NAMEEditorMode::UPLUGIN_NAMEEditorMode()
+{
+	FModuleManager::Get().LoadModule("EditorStyle");
+
+	// appearance and icon in the editing mode ribbon can be customized here
+	Info = FEditorModeInfo(UPLUGIN_NAMEEditorMode::EM_PLUGIN_NAMEEditorModeId,
+		LOCTEXT("ModeName", "PLUGIN_NAME"),
+		FSlateIcon(),
+		true);
+}
+
+
+UPLUGIN_NAMEEditorMode::~UPLUGIN_NAMEEditorMode()
+{
+}
+
+
+void UPLUGIN_NAMEEditorMode::ActorSelectionChangeNotify()
+{
+}
+
+void UPLUGIN_NAMEEditorMode::Enter()
+{
+	UEdMode::Enter();
+
+	//////////////////////////////////////////////////////////////////////////
+	//////////////////////////////////////////////////////////////////////////
+	// AddYourTool Step 2 - register the ToolBuilders for your Tools here.
+	// The string name you pass to the ToolManager is used to select/activate your ToolBuilder later.
+	//////////////////////////////////////////////////////////////////////////
+	////////////////////////////////////////////////////////////////////////// 
+	const FPLUGIN_NAMEEditorModeCommands& SampleToolCommands = FPLUGIN_NAMEEditorModeCommands::Get();
+
+	RegisterTool(SampleToolCommands.SimpleTool, SimpleToolName, NewObject<UPLUGIN_NAMESimpleToolBuilder>(this));
+	RegisterTool(SampleToolCommands.InteractiveTool, InteractiveToolName, NewObject<UPLUGIN_NAMEInteractiveToolBuilder>(this));
+
+	// active tool type is not relevant here, we just set to default
+	GetToolManager()->SelectActiveToolType(EToolSide::Left, SimpleToolName);
+}
+
+void UPLUGIN_NAMEEditorMode::CreateToolkit()
+{
+	Toolkit = MakeShareable(new FPLUGIN_NAMEEditorModeToolkit);
+}
+
+TMap<FName, TArray<TSharedPtr<FUICommandInfo>>> UPLUGIN_NAMEEditorMode::GetModeCommands() const
+{
+	return FPLUGIN_NAMEEditorModeCommands::Get().GetCommands();
+}
+
+#undef LOCTEXT_NAMESPACE
diff --git a/Engine/Plugins/Editor/PluginBrowser/Templates/EditorMode/Source/PLUGIN_NAME/Private/PLUGIN_NAMEEditorModeCommands.cpp b/Engine/Plugins/Editor/PluginBrowser/Templates/EditorMode/Source/PLUGIN_NAME/Private/PLUGIN_NAMEEditorModeCommands.cpp
new file mode 100644
--- /dev/null
+++ b/Engine/Plugins/Editor/PluginBrowser/Templates/EditorMode/Source/PLUGIN_NAME/Private/PLUGIN_NAMEEditorModeCommands.cpp
@@ -0,0 +1,33 @@
+// Copyright Epic Games, Inc. All Rights Reserved.
+
+#include "PLUGIN_NAMEEditorModeCommands.h"
+#include "PLUGIN_NAMEEditorMode.h"
+#include "EditorStyleSet.h"
+
+#define LOCTEXT_NAMESPACE "PLUGIN_NAMEEditorModeCommands"
+
+FPLUGIN_NAMEEditorModeCommands::FPLUGIN_NAMEEditorModeCommands()
+	: TCommands<FPLUGIN_NAMEEditorModeCommands>("PLUGIN_NAMEEditorMode",
+		NSLOCTEXT("PLUGIN_NAMEEditorMode", "PLUGIN_NAMEEditorModeCommands", "PLUGIN_NAME Editor Mode"),
+		NAME_None,
+		FEditorStyle::GetStyleSetName())
+{
+}
+
+void FPLUGIN_NAMEEditorModeCommands::RegisterCommands()
+{
+	TArray <TSharedPtr<FUICommandInfo>>& ToolCommands = Commands.FindOrAdd(NAME_Default);
+
+	UI_COMMAND(SimpleTool, "Show Actor Info", "Opens message box with info about a clicked actor", EUserInterfaceActionType::Button, FInputChord());
+	ToolCommands.Add(SimpleTool);
+
+	UI_COMMAND(InteractiveTool, "Measure Distance", "Measures distance between 2 points (click to set origin, shift-click to set end point)", EUserInterfaceActionType::ToggleButton, FInputChord());
+	ToolCommands.Add(InteractiveTool);
+}
+
+TMap<FName, TArray<TSharedPtr<FUICommandInfo>>> FPLUGIN_NAMEEditorModeCommands::GetCommands()
+{
+	return FPLUGIN_NAMEEditorModeCommands::Get().Commands;
+}
+
+#undef LOCTEXT_NAMESPACE
diff --git a/Engine/Plugins/Editor/PluginBrowser/Templates/EditorMode/Source/PLUGIN_NAME/Private/PLUGIN_NAMEEditorModeToolkit.cpp b/Engine/Plugins/Editor/PluginBrowser/Templates/EditorMode/Source/PLUGIN_NAME/Private/PLUGIN_NAMEEditorModeToolkit.cpp
new file mode 100644
--- /dev/null
+++ b/Engine/Plugins/Editor/PluginBrowser/Templates/EditorMode/Source/PLUGIN_NAME/Private/PLUGIN_NAMEEditorModeToolkit.cpp
@@ -0,0 +1,39 @@
+// Copyright Epic Games, Inc. All Rights Reserved.
+
+#include "PLUGIN_NAMEEditorModeToolkit.h"
+#include "PLUGIN_NAMEEditorMode.h"
+#include "Engine/Selection.h"
+
+#include "Modules/ModuleManager.h"
+#include "PropertyEditorModule.h"
+#include "IDetailsView.h"
+#include "EditorModeManager.h"
+
+#define LOCTEXT_NAMESPACE "PLUGIN_NAMEEditorModeToolkit"
+
+FPLUGIN_NAMEEditorModeToolkit::FPLUGIN_NAMEEditorModeToolkit()
+{
+}
+
+void FPLUGIN_NAMEEditorModeToolkit::Init(const TSharedPtr<IToolkitHost>& InitToolkitHost, TWeakObjectPtr<UEdMode> InOwningMode)
+{
+	FModeToolkit::Init(InitToolkitHost, InOwningMode);
+}
+
+void FPLUGIN_NAMEEditorModeToolkit::GetToolPaletteNames(TArray<FName>& PaletteNames) const
+{
+	PaletteNames.Add(NAME_Default);
+}
+
+
+FName FPLUGIN_NAMEEditorModeToolkit::GetToolkitFName() const
+{
+	return FName("PLUGIN_NAMEEditorMode");
+}
+
+FText FPLUGIN_NAMEEditorModeToolkit::GetBaseToolkitName() const
+{
+	return LOCTEXT("DisplayName", "PLUGIN_NAMEEditorMode Toolkit");
+}
+
+#undef LOCTEXT_NAMESPACE
diff --git a/Engine/Plugins/Editor/PluginBrowser/Templates/EditorMode/Source/PLUGIN_NAME/Private/PLUGIN_NAMEModule.cpp b/Engine/Plugins/Editor/PluginBrowser/Templates/EditorMode/Source/PLUGIN_NAME/Private/PLUGIN_NAMEModule.cpp
new file mode 100644
--- /dev/null
+++ b/Engine/Plugins/Editor/PluginBrowser/Templates/EditorMode/Source/PLUGIN_NAME/Private/PLUGIN_NAMEModule.cpp
@@ -0,0 +1,25 @@
+// Copyright Epic Games, Inc. All Rights Reserved.
+
+#include "PLUGIN_NAMEModule.h"
+#include "PLUGIN_NAMEEditorModeCommands.h"
+
+#define LOCTEXT_NAMESPACE "PLUGIN_NAMEModule"
+
+void FPLUGIN_NAMEModule::StartupModule()
+{
+	// This code will execute after your module is loaded into memory; the exact timing is specified in the .uplugin file per-module
+
+	FPLUGIN_NAMEEditorModeCommands::Register();
+}
+
+void FPLUGIN_NAMEModule::ShutdownModule()
+{
+	// This function may be called during shutdown to clean up your module.  For modules that support dynamic reloading,
+	// we call this function before unloading the module.
+
+	FPLUGIN_NAMEEditorModeCommands::Unregister();
+}
+
+#undef LOCTEXT_NAMESPACE
+
+IMPLEMENT_MODULE(FPLUGIN_NAMEModule, PLUGIN_NAMEEditorMode)
\ No newline at end of file
diff --git a/Engine/Plugins/Editor/PluginBrowser/Templates/EditorMode/Source/PLUGIN_NAME/Private/Tools/PLUGIN_NAMEInteractiveTool.cpp b/Engine/Plugins/Editor/PluginBrowser/Templates/EditorMode/Source/PLUGIN_NAME/Private/Tools/PLUGIN_NAMEInteractiveTool.cpp
new file mode 100644
--- /dev/null
+++ b/Engine/Plugins/Editor/PluginBrowser/Templates/EditorMode/Source/PLUGIN_NAME/Private/Tools/PLUGIN_NAMEInteractiveTool.cpp
@@ -0,0 +1,154 @@
+// Copyright Epic Games, Inc. All Rights Reserved.
+
+#include "PLUGIN_NAMEInteractiveTool.h"
+#include "InteractiveToolManager.h"
+#include "ToolBuilderUtil.h"
+#include "BaseBehaviors/ClickDragBehavior.h"
+
+// for raycast into World
+#include "CollisionQueryParams.h"
+#include "Engine/World.h"
+
+#include "SceneManagement.h"
+
+// localization namespace
+#define LOCTEXT_NAMESPACE "UPLUGIN_NAMEInteractiveTool"
+
+/*
+ * ToolBuilder
+ */
+
+UInteractiveTool* UPLUGIN_NAMEInteractiveToolBuilder::BuildTool(const FToolBuilderState & SceneState) const
+{
+	UPLUGIN_NAMEInteractiveTool* NewTool = NewObject<UPLUGIN_NAMEInteractiveTool>(SceneState.ToolManager);
+	NewTool->SetWorld(SceneState.World);
+	return NewTool;
+}
+
+
+// JAH TODO: update comments
+/*
+ * Tool
+ */
+
+UPLUGIN_NAMEInteractiveToolProperties::UPLUGIN_NAMEInteractiveToolProperties()
+{
+	// initialize the points and distance to reasonable values
+	StartPoint = FVector(0,0,0);
+	EndPoint = FVector(0,0,100);
+	Distance = 100;
+}
+
+
+void UPLUGIN_NAMEInteractiveTool::SetWorld(UWorld* World)
+{
+	check(World);
+	this->TargetWorld = World;
+}
+
+
+void UPLUGIN_NAMEInteractiveTool::Setup()
+{
+	UInteractiveTool::Setup();
+
+	// Add default mouse input behavior
+	UClickDragInputBehavior* MouseBehavior = NewObject<UClickDragInputBehavior>();
+	// We will use the shift key to indicate that we should move the second point. 
+	// This call tells the Behavior to call our OnUpdateModifierState() function on mouse-down and mouse-move
+	MouseBehavior->Modifiers.RegisterModifier(MoveSecondPointModifierID, FInputDeviceState::IsShiftKeyDown);
+	MouseBehavior->Initialize(this);
+	AddInputBehavior(MouseBehavior);
+
+	// Create the property set and register it with the Tool
+	Properties = NewObject<UPLUGIN_NAMEInteractiveToolProperties>(this, "Measurement");
+	AddToolPropertySource(Properties);
+	
+	bSecondPointModifierDown = false;
+	bMoveSecondPoint = false;
+}
+
+
+void UPLUGIN_NAMEInteractiveTool::OnUpdateModifierState(int ModifierID, bool bIsOn)
+{
+	// keep track of the "second point" modifier (shift key for mouse input)
+	if (ModifierID == MoveSecondPointModifierID)
+	{
+		bSecondPointModifierDown = bIsOn;
+	}
+}
+
+
+FInputRayHit UPLUGIN_NAMEInteractiveTool::CanBeginClickDragSequence(const FInputDeviceRay& PressPos)
+{
+	// we only start drag if press-down is on top of something we can raycast
+	FVector Temp;
+	FInputRayHit Result = FindRayHit(PressPos.WorldRay, Temp);
+	return Result;
+}
+
+
+void UPLUGIN_NAMEInteractiveTool::OnClickPress(const FInputDeviceRay& PressPos)
+{
+	// determine whether we are moving first or second point for the drag sequence
+	bMoveSecondPoint = bSecondPointModifierDown;
+	UpdatePosition(PressPos.WorldRay);
+}
+
+
+void UPLUGIN_NAMEInteractiveTool::OnClickDrag(const FInputDeviceRay& DragPos)
+{
+	UpdatePosition(DragPos.WorldRay);
+}
+
+
+FInputRayHit UPLUGIN_NAMEInteractiveTool::FindRayHit(const FRay& WorldRay, FVector& HitPos)
+{
+	// trace a ray into the World
+	FCollisionObjectQueryParams QueryParams(FCollisionObjectQueryParams::AllObjects);
+	FHitResult Result;
+	bool bHitWorld = TargetWorld->LineTraceSingleByObjectType(Result, WorldRay.Origin, WorldRay.PointAt(999999), QueryParams);
+	if (bHitWorld)
+	{
+		HitPos = Result.ImpactPoint;
+		return FInputRayHit(Result.Distance);
+	}
+	return FInputRayHit();
+}
+
+
+void UPLUGIN_NAMEInteractiveTool::UpdatePosition(const FRay& WorldRay)
+{
+	FInputRayHit HitResult = FindRayHit(WorldRay, (bMoveSecondPoint) ? Properties->EndPoint : Properties->StartPoint);
+	if (HitResult.bHit)
+	{
+		UpdateDistance();
+	}
+}
+
+
+void UPLUGIN_NAMEInteractiveTool::UpdateDistance()
+{
+	Properties->Distance = FVector::Distance(Properties->StartPoint, Properties->EndPoint);
+}
+
+
+void UPLUGIN_NAMEInteractiveTool::OnPropertyModified(UObject* PropertySet, FProperty* Property)
+{
+	// if the user updated any of the property fields, update the distance
+	UpdateDistance();
+}
+
+
+void UPLUGIN_NAMEInteractiveTool::Render(IToolsContextRenderAPI* RenderAPI)
+{
+	FPrimitiveDrawInterface* PDI = RenderAPI->GetPrimitiveDrawInterface();
+	// draw a thin line that shows through objects
+	PDI->DrawLine(Properties->StartPoint, Properties->EndPoint,
+		FColor(240, 16, 16), SDPG_Foreground, 2.0f, 0.0f, true);
+	// draw a thicker line that is depth-tested
+	PDI->DrawLine(Properties->StartPoint, Properties->EndPoint,
+		FColor(240, 16, 16), SDPG_World, 4.0f, 0.0f, true);
+}
+
+
+#undef LOCTEXT_NAMESPACE
diff --git a/Engine/Plugins/Editor/PluginBrowser/Templates/EditorMode/Source/PLUGIN_NAME/Private/Tools/PLUGIN_NAMEInteractiveTool.h b/Engine/Plugins/Editor/PluginBrowser/Templates/EditorMode/Source/PLUGIN_NAME/Private/Tools/PLUGIN_NAMEInteractiveTool.h
new file mode 100644
--- /dev/null
+++ b/Engine/Plugins/Editor/PluginBrowser/Templates/EditorMode/Source/PLUGIN_NAME/Private/Tools/PLUGIN_NAMEInteractiveTool.h
@@ -0,0 +1,97 @@
+// Copyright Epic Games, Inc. All Rights Reserved.
+
+#pragma once
+
+#include "CoreMinimal.h"
+#include "InteractiveToolBuilder.h"
+#include "BaseTools/ClickDragTool.h"
+#include "PLUGIN_NAMEInteractiveTool.generated.h"
+
+
+/**
+ * Builder for UPLUGIN_NAMEInteractiveTool
+ */
+UCLASS()
+class PLUGIN_NAME_API UPLUGIN_NAMEInteractiveToolBuilder : public UInteractiveToolBuilder
+{
+	GENERATED_BODY()
+
+public:
+	virtual bool CanBuildTool(const FToolBuilderState& SceneState) const override { return true; }
+	virtual UInteractiveTool* BuildTool(const FToolBuilderState& SceneState) const override;
+};
+
+
+/**
+ * Property set for the UPLUGIN_NAMEInteractiveTool
+ */
+UCLASS(Transient)
+class PLUGIN_NAME_API UPLUGIN_NAMEInteractiveToolProperties : public UInteractiveToolPropertySet
+{
+	GENERATED_BODY()
+
+public:
+	UPLUGIN_NAMEInteractiveToolProperties();
+
+	/** First point of measurement */
+	UPROPERTY(EditAnywhere, Category = Options)
+	FVector StartPoint;
+
+	/** Second point of measurement */
+	UPROPERTY(EditAnywhere, Category = Options)
+	FVector EndPoint;
+	
+	/** Current distance measurement */
+	UPROPERTY(EditAnywhere, Category = Options)
+	double Distance;
+};
+
+
+
+/**
+ * UPLUGIN_NAMEInteractiveTool is an example Tool that allows the user to measure the 
+ * distance between two points. The first point is set by click-dragging the mouse, and
+ * the second point is set by shift-click-dragging the mouse.
+ */
+UCLASS()
+class PLUGIN_NAME_API UPLUGIN_NAMEInteractiveTool : public UInteractiveTool, public IClickDragBehaviorTarget
+{
+	GENERATED_BODY()
+
+public:
+	virtual void SetWorld(UWorld* World);
+
+	/** UInteractiveTool overrides */
+	virtual void Setup() override;
+	virtual void Render(IToolsContextRenderAPI* RenderAPI) override;
+	virtual void OnPropertyModified(UObject* PropertySet, FProperty* Property) override;
+
+	/** IClickDragBehaviorTarget implementation */
+	virtual FInputRayHit CanBeginClickDragSequence(const FInputDeviceRay& PressPos) override;
+	virtual void OnClickPress(const FInputDeviceRay& PressPos) override;
+	virtual void OnClickDrag(const FInputDeviceRay& DragPos) override;
+	// these are not used in this Tool
+	virtual void OnClickRelease(const FInputDeviceRay& ReleasePos) override {}
+	virtual void OnTerminateDragSequence() override {}
+
+	/** IModifierToggleBehaviorTarget implementation (inherited via IClickDragBehaviorTarget) */
+	virtual void OnUpdateModifierState(int ModifierID, bool bIsOn) override;
+
+
+protected:
+	/** Properties of the tool are stored here */
+	UPROPERTY()
+	UPLUGIN_NAMEInteractiveToolProperties* Properties;
+
+
+protected:
+	UWorld* TargetWorld = nullptr;		// target World we will raycast into
+
+	static const int MoveSecondPointModifierID = 1;		// identifier we associate with the shift key
+	bool bSecondPointModifierDown = false;				// flag we use to keep track of modifier state
+	bool bMoveSecondPoint = false;						// flag we use to keep track of which point we are moving during a press-drag
+
+	FInputRayHit FindRayHit(const FRay& WorldRay, FVector& HitPos);		// raycasts into World
+	void UpdatePosition(const FRay& WorldRay);					// updates first or second point based on raycast
+	void UpdateDistance();										// updates distance
+};
diff --git a/Engine/Plugins/Editor/PluginBrowser/Templates/EditorMode/Source/PLUGIN_NAME/Private/Tools/PLUGIN_NAMESimpleTool.cpp b/Engine/Plugins/Editor/PluginBrowser/Templates/EditorMode/Source/PLUGIN_NAME/Private/Tools/PLUGIN_NAMESimpleTool.cpp
new file mode 100644
--- /dev/null
+++ b/Engine/Plugins/Editor/PluginBrowser/Templates/EditorMode/Source/PLUGIN_NAME/Private/Tools/PLUGIN_NAMESimpleTool.cpp
@@ -0,0 +1,95 @@
+// Copyright Epic Games, Inc. All Rights Reserved.
+
+#include "PLUGIN_NAMESimpleTool.h"
+#include "InteractiveToolManager.h"
+#include "ToolBuilderUtil.h"
+#include "CollisionQueryParams.h"
+#include "Engine/World.h"
+
+// localization namespace
+#define LOCTEXT_NAMESPACE "PLUGIN_NAMESimpleTool"
+
+/*
+ * ToolBuilder implementation
+ */
+
+UInteractiveTool* UPLUGIN_NAMESimpleToolBuilder::BuildTool(const FToolBuilderState& SceneState) const
+{
+	UPLUGIN_NAMESimpleTool* NewTool = NewObject<UPLUGIN_NAMESimpleTool>(SceneState.ToolManager);
+	NewTool->SetWorld(SceneState.World);
+	return NewTool;
+}
+
+
+
+/*
+ * ToolProperties implementation
+ */
+
+UPLUGIN_NAMESimpleToolProperties::UPLUGIN_NAMESimpleToolProperties()
+{
+	ShowExtendedInfo = true;
+}
+
+
+/*
+ * Tool implementation
+ */
+
+UPLUGIN_NAMESimpleTool::UPLUGIN_NAMESimpleTool()
+{
+}
+
+
+void UPLUGIN_NAMESimpleTool::SetWorld(UWorld* World)
+{
+	this->TargetWorld = World;
+}
+
+
+void UPLUGIN_NAMESimpleTool::Setup()
+{
+	USingleClickTool::Setup();
+
+	Properties = NewObject<UPLUGIN_NAMESimpleToolProperties>(this);
+	AddToolPropertySource(Properties);
+}
+
+
+void UPLUGIN_NAMESimpleTool::OnClicked(const FInputDeviceRay& ClickPos)
+{
+	// we will create actor at this position
+	FVector NewActorPos = FVector::ZeroVector;
+
+	// cast ray into world to find hit position
+	FVector RayStart = ClickPos.WorldRay.Origin;
+	FVector RayEnd = ClickPos.WorldRay.PointAt(99999999.f);
+	FCollisionObjectQueryParams QueryParams(FCollisionObjectQueryParams::AllObjects);
+	FHitResult Result;
+	if (TargetWorld->LineTraceSingleByObjectType(Result, RayStart, RayEnd, QueryParams))
+	{
+		if (AActor* ClickedActor = Result.GetActor())
+		{
+			FText ActorInfoMsg;
+
+			if (Properties->ShowExtendedInfo)
+			{
+				ActorInfoMsg = FText::Format(LOCTEXT("ExtendedActorInfo", "Name: {0}\nClass: {1}"), 
+					FText::FromString(ClickedActor->GetName()), 
+					FText::FromString(ClickedActor->GetClass()->GetName())
+				);
+			}
+			else
+			{
+				ActorInfoMsg = FText::Format(LOCTEXT("BasicActorInfo", "Name: {0}"), FText::FromString(Result.GetActor()->GetName()));
+			}
+
+			FText Title = LOCTEXT("ActorInfoDialogTitle", "Actor Info");
+			// JAH TODO: consider if we can highlight the actor prior to opening the dialog box or make it non-modal
+			FMessageDialog::Open(EAppMsgType::Ok, ActorInfoMsg, &Title);
+		}
+	}
+}
+
+
+#undef LOCTEXT_NAMESPACE
\ No newline at end of file
diff --git a/Engine/Plugins/Editor/PluginBrowser/Templates/EditorMode/Source/PLUGIN_NAME/Private/Tools/PLUGIN_NAMESimpleTool.h b/Engine/Plugins/Editor/PluginBrowser/Templates/EditorMode/Source/PLUGIN_NAME/Private/Tools/PLUGIN_NAMESimpleTool.h
new file mode 100644
--- /dev/null
+++ b/Engine/Plugins/Editor/PluginBrowser/Templates/EditorMode/Source/PLUGIN_NAME/Private/Tools/PLUGIN_NAMESimpleTool.h
@@ -0,0 +1,72 @@
+// Copyright Epic Games, Inc. All Rights Reserved.
+
+#pragma once
+
+#include "CoreMinimal.h"
+#include "UObject/NoExportTypes.h"
+#include "InteractiveToolBuilder.h"
+#include "BaseTools/SingleClickTool.h"
+#include "PLUGIN_NAMESimpleTool.generated.h"
+
+/**
+ * Builder for UPLUGIN_NAMESimpleTool
+ */
+UCLASS()
+class PLUGIN_NAME_API UPLUGIN_NAMESimpleToolBuilder : public UInteractiveToolBuilder
+{
+	GENERATED_BODY()
+
+public:
+	virtual bool CanBuildTool(const FToolBuilderState& SceneState) const override { return true; }
+	virtual UInteractiveTool* BuildTool(const FToolBuilderState& SceneState) const override;
+};
+
+
+
+/**
+ * Settings UObject for UPLUGIN_NAMESimpleTool. This UClass inherits from UInteractiveToolPropertySet,
+ * which provides an OnModified delegate that the Tool will listen to for changes in property values.
+ */
+UCLASS(Transient)
+class PLUGIN_NAME_API UPLUGIN_NAMESimpleToolProperties : public UInteractiveToolPropertySet
+{
+	GENERATED_BODY()
+public:
+	UPLUGIN_NAMESimpleToolProperties();
+
+	/** If enabled, dialog should display extended information about the actor clicked on. Otherwise, only basic info will be shown. */
+	UPROPERTY(EditAnywhere, Category = Options, meta = (DisplayName = "Show Extended Info"))
+	bool ShowExtendedInfo;
+};
+
+
+
+
+/**
+ * UPLUGIN_NAMESimpleTool is an example Tool that opens a message box displaying info about an actor that the user
+ * clicks left mouse button. All the action is in the ::OnClicked handler.
+ */
+UCLASS()
+class PLUGIN_NAME_API UPLUGIN_NAMESimpleTool : public USingleClickTool
+{
+	GENERATED_BODY()
+
+public:
+	UPLUGIN_NAMESimpleTool();
+
+	virtual void SetWorld(UWorld* World);
+
+	virtual void Setup() override;
+
+	virtual void OnClicked(const FInputDeviceRay& ClickPos);
+
+
+protected:
+	UPROPERTY()
+	UPLUGIN_NAMESimpleToolProperties* Properties;
+
+
+protected:
+	/** target World we will raycast into to find actors */
+	UWorld* TargetWorld;
+};
\ No newline at end of file
diff --git a/Engine/Plugins/Editor/PluginBrowser/Templates/EditorMode/Source/PLUGIN_NAME/Public/PLUGIN_NAME.h b/Engine/Plugins/Editor/PluginBrowser/Templates/EditorMode/Source/PLUGIN_NAME/Public/PLUGIN_NAME.h
deleted file mode 100644
diff --git a/Engine/Plugins/Editor/PluginBrowser/Templates/EditorMode/Source/PLUGIN_NAME/Public/PLUGIN_NAMEEdMode.h b/Engine/Plugins/Editor/PluginBrowser/Templates/EditorMode/Source/PLUGIN_NAME/Public/PLUGIN_NAMEEdMode.h
deleted file mode 100644
diff --git a/Engine/Plugins/Editor/PluginBrowser/Templates/EditorMode/Source/PLUGIN_NAME/Public/PLUGIN_NAMEEdModeToolkit.h b/Engine/Plugins/Editor/PluginBrowser/Templates/EditorMode/Source/PLUGIN_NAME/Public/PLUGIN_NAMEEdModeToolkit.h
deleted file mode 100644
diff --git a/Engine/Plugins/Editor/PluginBrowser/Templates/EditorMode/Source/PLUGIN_NAME/Public/PLUGIN_NAMEEditorMode.h b/Engine/Plugins/Editor/PluginBrowser/Templates/EditorMode/Source/PLUGIN_NAME/Public/PLUGIN_NAMEEditorMode.h
new file mode 100644
--- /dev/null
+++ b/Engine/Plugins/Editor/PluginBrowser/Templates/EditorMode/Source/PLUGIN_NAME/Public/PLUGIN_NAMEEditorMode.h
@@ -0,0 +1,36 @@
+// Copyright Epic Games, Inc. All Rights Reserved.
+
+#pragma once
+
+#include "CoreMinimal.h"
+#include "Tools/UEdMode.h"
+#include "PLUGIN_NAMEEditorMode.generated.h"
+
+/**
+ * This class provides an example of how to extend a UEdMode to add some simple tools
+ * using the InteractiveTools framework. The various UEdMode input event handlers (see UEdMode.h)
+ * forward events to a UEdModeInteractiveToolsContext instance, which
+ * has all the logic for interacting with the InputRouter, ToolManager, etc.
+ * The functions provided here are the minimum to get started inserting some custom behavior.
+ * Take a look at the UEdMode markup for more extensibility options.
+ */
+UCLASS()
+class UPLUGIN_NAMEEditorMode : public UEdMode
+{
+	GENERATED_BODY()
+
+public:
+	const static FEditorModeID EM_PLUGIN_NAMEEditorModeId;
+
+	static FString SimpleToolName;
+	static FString InteractiveToolName;
+
+	UPLUGIN_NAMEEditorMode();
+	virtual ~UPLUGIN_NAMEEditorMode();
+
+	/** UEdMode interface */
+	virtual void Enter() override;
+	virtual void ActorSelectionChangeNotify() override;
+	virtual void CreateToolkit() override;
+	virtual TMap<FName, TArray<TSharedPtr<FUICommandInfo>>> GetModeCommands() const override;
+};
diff --git a/Engine/Plugins/Editor/PluginBrowser/Templates/EditorMode/Source/PLUGIN_NAME/Public/PLUGIN_NAMEEditorModeCommands.h b/Engine/Plugins/Editor/PluginBrowser/Templates/EditorMode/Source/PLUGIN_NAME/Public/PLUGIN_NAMEEditorModeCommands.h
new file mode 100644
--- /dev/null
+++ b/Engine/Plugins/Editor/PluginBrowser/Templates/EditorMode/Source/PLUGIN_NAME/Public/PLUGIN_NAMEEditorModeCommands.h
@@ -0,0 +1,24 @@
+// Copyright Epic Games, Inc. All Rights Reserved.
+
+#pragma once
+
+#include "CoreMinimal.h"
+#include "Framework/Commands/Commands.h"
+
+/**
+ * This class contains info about the full set of commands used in this editor mode.
+ */
+class FPLUGIN_NAMEEditorModeCommands : public TCommands<FPLUGIN_NAMEEditorModeCommands>
+{
+public:
+	FPLUGIN_NAMEEditorModeCommands();
+
+	virtual void RegisterCommands() override;
+	static TMap<FName, TArray<TSharedPtr<FUICommandInfo>>> GetCommands();
+
+	TSharedPtr<FUICommandInfo> SimpleTool;
+	TSharedPtr<FUICommandInfo> InteractiveTool;
+
+protected:
+	TMap<FName, TArray<TSharedPtr<FUICommandInfo>>> Commands;
+};
diff --git a/Engine/Plugins/Editor/PluginBrowser/Templates/EditorMode/Source/PLUGIN_NAME/Public/PLUGIN_NAMEEditorModeToolkit.h b/Engine/Plugins/Editor/PluginBrowser/Templates/EditorMode/Source/PLUGIN_NAME/Public/PLUGIN_NAMEEditorModeToolkit.h
new file mode 100644
--- /dev/null
+++ b/Engine/Plugins/Editor/PluginBrowser/Templates/EditorMode/Source/PLUGIN_NAME/Public/PLUGIN_NAMEEditorModeToolkit.h
@@ -0,0 +1,25 @@
+// Copyright Epic Games, Inc. All Rights Reserved.
+
+#pragma once
+
+#include "CoreMinimal.h"
+#include "Toolkits/BaseToolkit.h"
+#include "PLUGIN_NAMEEditorMode.h"
+
+/**
+ * This FModeToolkit just creates a basic UI panel that allows various InteractiveTools to
+ * be initialized, and a DetailsView used to show properties of the active Tool.
+ */
+class FPLUGIN_NAMEEditorModeToolkit : public FModeToolkit
+{
+public:
+	FPLUGIN_NAMEEditorModeToolkit();
+
+	/** FModeToolkit interface */
+	virtual void Init(const TSharedPtr<IToolkitHost>& InitToolkitHost, TWeakObjectPtr<UEdMode> InOwningMode) override;
+	virtual void GetToolPaletteNames(TArray<FName>& PaletteNames) const override;
+
+	/** IToolkit interface */
+	virtual FName GetToolkitFName() const override;
+	virtual FText GetBaseToolkitName() const override;
+};
diff --git a/Engine/Plugins/Editor/PluginBrowser/Templates/EditorMode/Source/PLUGIN_NAME/Public/PLUGIN_NAMEModule.h b/Engine/Plugins/Editor/PluginBrowser/Templates/EditorMode/Source/PLUGIN_NAME/Public/PLUGIN_NAMEModule.h
new file mode 100644
--- /dev/null
+++ b/Engine/Plugins/Editor/PluginBrowser/Templates/EditorMode/Source/PLUGIN_NAME/Public/PLUGIN_NAMEModule.h
@@ -0,0 +1,19 @@
+// Copyright Epic Games, Inc. All Rights Reserved.
+
+#pragma once
+
+#include "CoreMinimal.h"
+#include "Modules/ModuleManager.h"
+
+/**
+ * This is the module definition for the editor mode. You can implement custom functionality
+ * as your plugin module starts up and shuts down. See IModuleInterface for more extensibility options.
+ */
+class FPLUGIN_NAMEModule : public IModuleInterface
+{
+public:
+
+	/** IModuleInterface implementation */
+	virtual void StartupModule() override;
+	virtual void ShutdownModule() override;
+};
