Updated raw object pointers to TObjectPtr in the EditorMode plugin template.

diff --git a/Engine/Plugins/Editor/PluginBrowser/Templates/EditorMode/Source/PLUGIN_NAME/Private/Tools/PLUGIN_NAMEInteractiveTool.h b/Engine/Plugins/Editor/PluginBrowser/Templates/EditorMode/Source/PLUGIN_NAME/Private/Tools/PLUGIN_NAMEInteractiveTool.h
--- a/Engine/Plugins/Editor/PluginBrowser/Templates/EditorMode/Source/PLUGIN_NAME/Private/Tools/PLUGIN_NAMEInteractiveTool.h
+++ b/Engine/Plugins/Editor/PluginBrowser/Templates/EditorMode/Source/PLUGIN_NAME/Private/Tools/PLUGIN_NAMEInteractiveTool.h
@@ -81,7 +81,7 @@ public:
 protected:
 	/** Properties of the tool are stored here */
 	UPROPERTY()
-	UPLUGIN_NAMEInteractiveToolProperties* Properties;
+	TObjectPtr<UPLUGIN_NAMEInteractiveToolProperties> Properties;
 
 
 protected:
diff --git a/Engine/Plugins/Editor/PluginBrowser/Templates/EditorMode/Source/PLUGIN_NAME/Private/Tools/PLUGIN_NAMESimpleTool.h b/Engine/Plugins/Editor/PluginBrowser/Templates/EditorMode/Source/PLUGIN_NAME/Private/Tools/PLUGIN_NAMESimpleTool.h
--- a/Engine/Plugins/Editor/PluginBrowser/Templates/EditorMode/Source/PLUGIN_NAME/Private/Tools/PLUGIN_NAMESimpleTool.h
+++ b/Engine/Plugins/Editor/PluginBrowser/Templates/EditorMode/Source/PLUGIN_NAME/Private/Tools/PLUGIN_NAMESimpleTool.h
@@ -63,7 +63,7 @@ public:
 
 protected:
 	UPROPERTY()
-	UPLUGIN_NAMESimpleToolProperties* Properties;
+	TObjectPtr<UPLUGIN_NAMESimpleToolProperties> Properties;
 
 
 protected:
