Remove CoreMinimal.h includes from plugin template headers

#tests manually created plugins with each template except game feature

diff --git a/Engine/Plugins/Editor/PluginBrowser/Templates/Advanced/Source/PLUGIN_NAME/Public/PLUGIN_NAME.h b/Engine/Plugins/Editor/PluginBrowser/Templates/Advanced/Source/PLUGIN_NAME/Public/PLUGIN_NAME.h
--- a/Engine/Plugins/Editor/PluginBrowser/Templates/Advanced/Source/PLUGIN_NAME/Public/PLUGIN_NAME.h
+++ b/Engine/Plugins/Editor/PluginBrowser/Templates/Advanced/Source/PLUGIN_NAME/Public/PLUGIN_NAME.h
@@ -2,7 +2,6 @@
 
 #pragma once
 
-#include "CoreMinimal.h"
 #include "Modules/ModuleManager.h"
 
 class FToolBarBuilder;
diff --git a/Engine/Plugins/Editor/PluginBrowser/Templates/Advanced/Source/PLUGIN_NAME/Public/PLUGIN_NAMECommands.h b/Engine/Plugins/Editor/PluginBrowser/Templates/Advanced/Source/PLUGIN_NAME/Public/PLUGIN_NAMECommands.h
--- a/Engine/Plugins/Editor/PluginBrowser/Templates/Advanced/Source/PLUGIN_NAME/Public/PLUGIN_NAMECommands.h
+++ b/Engine/Plugins/Editor/PluginBrowser/Templates/Advanced/Source/PLUGIN_NAME/Public/PLUGIN_NAMECommands.h
@@ -2,7 +2,6 @@
 
 #pragma once
 
-#include "CoreMinimal.h"
 #include "Framework/Commands/Commands.h"
 #include "PLUGIN_NAMEStyle.h"
 
diff --git a/Engine/Plugins/Editor/PluginBrowser/Templates/Advanced/Source/PLUGIN_NAME/Public/PLUGIN_NAMEStyle.h b/Engine/Plugins/Editor/PluginBrowser/Templates/Advanced/Source/PLUGIN_NAME/Public/PLUGIN_NAMEStyle.h
--- a/Engine/Plugins/Editor/PluginBrowser/Templates/Advanced/Source/PLUGIN_NAME/Public/PLUGIN_NAMEStyle.h
+++ b/Engine/Plugins/Editor/PluginBrowser/Templates/Advanced/Source/PLUGIN_NAME/Public/PLUGIN_NAMEStyle.h
@@ -2,7 +2,6 @@
 
 #pragma once
 
-#include "CoreMinimal.h"
 #include "Styling/SlateStyle.h"
 
 /**  */
diff --git a/Engine/Plugins/Editor/PluginBrowser/Templates/Basic/Source/PLUGIN_NAME/Public/PLUGIN_NAME.h b/Engine/Plugins/Editor/PluginBrowser/Templates/Basic/Source/PLUGIN_NAME/Public/PLUGIN_NAME.h
--- a/Engine/Plugins/Editor/PluginBrowser/Templates/Basic/Source/PLUGIN_NAME/Public/PLUGIN_NAME.h
+++ b/Engine/Plugins/Editor/PluginBrowser/Templates/Basic/Source/PLUGIN_NAME/Public/PLUGIN_NAME.h
@@ -2,7 +2,6 @@
 
 #pragma once
 
-#include "CoreMinimal.h"
 #include "Modules/ModuleManager.h"
 
 class FToolBarBuilder;
diff --git a/Engine/Plugins/Editor/PluginBrowser/Templates/Basic/Source/PLUGIN_NAME/Public/PLUGIN_NAMECommands.h b/Engine/Plugins/Editor/PluginBrowser/Templates/Basic/Source/PLUGIN_NAME/Public/PLUGIN_NAMECommands.h
--- a/Engine/Plugins/Editor/PluginBrowser/Templates/Basic/Source/PLUGIN_NAME/Public/PLUGIN_NAMECommands.h
+++ b/Engine/Plugins/Editor/PluginBrowser/Templates/Basic/Source/PLUGIN_NAME/Public/PLUGIN_NAMECommands.h
@@ -2,7 +2,6 @@
 
 #pragma once
 
-#include "CoreMinimal.h"
 #include "Framework/Commands/Commands.h"
 #include "PLUGIN_NAMEStyle.h"
 
diff --git a/Engine/Plugins/Editor/PluginBrowser/Templates/Basic/Source/PLUGIN_NAME/Public/PLUGIN_NAMEStyle.h b/Engine/Plugins/Editor/PluginBrowser/Templates/Basic/Source/PLUGIN_NAME/Public/PLUGIN_NAMEStyle.h
--- a/Engine/Plugins/Editor/PluginBrowser/Templates/Basic/Source/PLUGIN_NAME/Public/PLUGIN_NAMEStyle.h
+++ b/Engine/Plugins/Editor/PluginBrowser/Templates/Basic/Source/PLUGIN_NAME/Public/PLUGIN_NAMEStyle.h
@@ -2,7 +2,6 @@
 
 #pragma once
 
-#include "CoreMinimal.h"
 #include "Styling/SlateStyle.h"
 
 class FPLUGIN_NAMEStyle
diff --git a/Engine/Plugins/Editor/PluginBrowser/Templates/Blank/Source/PLUGIN_NAME/Public/PLUGIN_NAME.h b/Engine/Plugins/Editor/PluginBrowser/Templates/Blank/Source/PLUGIN_NAME/Public/PLUGIN_NAME.h
--- a/Engine/Plugins/Editor/PluginBrowser/Templates/Blank/Source/PLUGIN_NAME/Public/PLUGIN_NAME.h
+++ b/Engine/Plugins/Editor/PluginBrowser/Templates/Blank/Source/PLUGIN_NAME/Public/PLUGIN_NAME.h
@@ -2,7 +2,6 @@
 
 #pragma once
 
-#include "CoreMinimal.h"
 #include "Modules/ModuleManager.h"
 
 class FPLUGIN_NAMEModule : public IModuleInterface
diff --git a/Engine/Plugins/Editor/PluginBrowser/Templates/EditorMode/Source/PLUGIN_NAME/Private/Tools/PLUGIN_NAMEInteractiveTool.h b/Engine/Plugins/Editor/PluginBrowser/Templates/EditorMode/Source/PLUGIN_NAME/Private/Tools/PLUGIN_NAMEInteractiveTool.h
--- a/Engine/Plugins/Editor/PluginBrowser/Templates/EditorMode/Source/PLUGIN_NAME/Private/Tools/PLUGIN_NAMEInteractiveTool.h
+++ b/Engine/Plugins/Editor/PluginBrowser/Templates/EditorMode/Source/PLUGIN_NAME/Private/Tools/PLUGIN_NAMEInteractiveTool.h
@@ -2,7 +2,6 @@
 
 #pragma once
 
-#include "CoreMinimal.h"
 #include "InteractiveToolBuilder.h"
 #include "BaseTools/ClickDragTool.h"
 #include "PLUGIN_NAMEInteractiveTool.generated.h"
diff --git a/Engine/Plugins/Editor/PluginBrowser/Templates/EditorMode/Source/PLUGIN_NAME/Private/Tools/PLUGIN_NAMESimpleTool.h b/Engine/Plugins/Editor/PluginBrowser/Templates/EditorMode/Source/PLUGIN_NAME/Private/Tools/PLUGIN_NAMESimpleTool.h
--- a/Engine/Plugins/Editor/PluginBrowser/Templates/EditorMode/Source/PLUGIN_NAME/Private/Tools/PLUGIN_NAMESimpleTool.h
+++ b/Engine/Plugins/Editor/PluginBrowser/Templates/EditorMode/Source/PLUGIN_NAME/Private/Tools/PLUGIN_NAMESimpleTool.h
@@ -2,7 +2,6 @@
 
 #pragma once
 
-#include "CoreMinimal.h"
 #include "UObject/NoExportTypes.h"
 #include "InteractiveToolBuilder.h"
 #include "BaseTools/SingleClickTool.h"
diff --git a/Engine/Plugins/Editor/PluginBrowser/Templates/EditorMode/Source/PLUGIN_NAME/Public/PLUGIN_NAMEEditorMode.h b/Engine/Plugins/Editor/PluginBrowser/Templates/EditorMode/Source/PLUGIN_NAME/Public/PLUGIN_NAMEEditorMode.h
--- a/Engine/Plugins/Editor/PluginBrowser/Templates/EditorMode/Source/PLUGIN_NAME/Public/PLUGIN_NAMEEditorMode.h
+++ b/Engine/Plugins/Editor/PluginBrowser/Templates/EditorMode/Source/PLUGIN_NAME/Public/PLUGIN_NAMEEditorMode.h
@@ -2,7 +2,6 @@
 
 #pragma once
 
-#include "CoreMinimal.h"
 #include "Tools/UEdMode.h"
 #include "PLUGIN_NAMEEditorMode.generated.h"
 
diff --git a/Engine/Plugins/Editor/PluginBrowser/Templates/EditorMode/Source/PLUGIN_NAME/Public/PLUGIN_NAMEEditorModeCommands.h b/Engine/Plugins/Editor/PluginBrowser/Templates/EditorMode/Source/PLUGIN_NAME/Public/PLUGIN_NAMEEditorModeCommands.h
--- a/Engine/Plugins/Editor/PluginBrowser/Templates/EditorMode/Source/PLUGIN_NAME/Public/PLUGIN_NAMEEditorModeCommands.h
+++ b/Engine/Plugins/Editor/PluginBrowser/Templates/EditorMode/Source/PLUGIN_NAME/Public/PLUGIN_NAMEEditorModeCommands.h
@@ -2,7 +2,6 @@
 
 #pragma once
 
-#include "CoreMinimal.h"
 #include "Framework/Commands/Commands.h"
 
 /**
diff --git a/Engine/Plugins/Editor/PluginBrowser/Templates/EditorMode/Source/PLUGIN_NAME/Public/PLUGIN_NAMEEditorModeToolkit.h b/Engine/Plugins/Editor/PluginBrowser/Templates/EditorMode/Source/PLUGIN_NAME/Public/PLUGIN_NAMEEditorModeToolkit.h
--- a/Engine/Plugins/Editor/PluginBrowser/Templates/EditorMode/Source/PLUGIN_NAME/Public/PLUGIN_NAMEEditorModeToolkit.h
+++ b/Engine/Plugins/Editor/PluginBrowser/Templates/EditorMode/Source/PLUGIN_NAME/Public/PLUGIN_NAMEEditorModeToolkit.h
@@ -2,7 +2,6 @@
 
 #pragma once
 
-#include "CoreMinimal.h"
 #include "Toolkits/BaseToolkit.h"
 #include "PLUGIN_NAMEEditorMode.h"
 
diff --git a/Engine/Plugins/Editor/PluginBrowser/Templates/EditorMode/Source/PLUGIN_NAME/Public/PLUGIN_NAMEModule.h b/Engine/Plugins/Editor/PluginBrowser/Templates/EditorMode/Source/PLUGIN_NAME/Public/PLUGIN_NAMEModule.h
--- a/Engine/Plugins/Editor/PluginBrowser/Templates/EditorMode/Source/PLUGIN_NAME/Public/PLUGIN_NAMEModule.h
+++ b/Engine/Plugins/Editor/PluginBrowser/Templates/EditorMode/Source/PLUGIN_NAME/Public/PLUGIN_NAMEModule.h
@@ -2,7 +2,6 @@
 
 #pragma once
 
-#include "CoreMinimal.h"
 #include "Modules/ModuleManager.h"
 
 /**
