[Editor Mode Plugin Template] Fixed a bug of a new plugin not creating with the Editor Plugin template and non-unity mode is enabled / general cleanup

diff --git a/Engine/Plugins/Editor/PluginBrowser/Templates/EditorMode/Source/PLUGIN_NAME/Private/PLUGIN_NAMEEditorMode.cpp b/Engine/Plugins/Editor/PluginBrowser/Templates/EditorMode/Source/PLUGIN_NAME/Private/PLUGIN_NAMEEditorMode.cpp
--- a/Engine/Plugins/Editor/PluginBrowser/Templates/EditorMode/Source/PLUGIN_NAME/Private/PLUGIN_NAMEEditorMode.cpp
+++ b/Engine/Plugins/Editor/PluginBrowser/Templates/EditorMode/Source/PLUGIN_NAME/Private/PLUGIN_NAMEEditorMode.cpp
@@ -5,6 +5,7 @@
 #include "EdModeInteractiveToolsContext.h"
 #include "InteractiveToolManager.h"
 #include "PLUGIN_NAMEEditorModeCommands.h"
+#include "Modules/ModuleManager.h"
 
 
 //////////////////////////////////////////////////////////////////////////
diff --git a/Engine/Plugins/Editor/PluginBrowser/Templates/EditorMode/Source/PLUGIN_NAME/Private/PLUGIN_NAMEEditorModeCommands.cpp b/Engine/Plugins/Editor/PluginBrowser/Templates/EditorMode/Source/PLUGIN_NAME/Private/PLUGIN_NAMEEditorModeCommands.cpp
--- a/Engine/Plugins/Editor/PluginBrowser/Templates/EditorMode/Source/PLUGIN_NAME/Private/PLUGIN_NAMEEditorModeCommands.cpp
+++ b/Engine/Plugins/Editor/PluginBrowser/Templates/EditorMode/Source/PLUGIN_NAME/Private/PLUGIN_NAMEEditorModeCommands.cpp
@@ -10,7 +10,7 @@ FPLUGIN_NAMEEditorModeCommands::FPLUGIN_NAMEEditorModeCommands()
 	: TCommands<FPLUGIN_NAMEEditorModeCommands>("PLUGIN_NAMEEditorMode",
 		NSLOCTEXT("PLUGIN_NAMEEditorMode", "PLUGIN_NAMEEditorModeCommands", "PLUGIN_NAME Editor Mode"),
 		NAME_None,
-		FEditorStyle::GetStyleSetName())
+		FAppStyle::GetAppStyleSetName())
 {
 }
 
diff --git a/Engine/Plugins/Editor/PluginBrowser/Templates/EditorMode/Source/PLUGIN_NAME/Private/PLUGIN_NAMEModule.cpp b/Engine/Plugins/Editor/PluginBrowser/Templates/EditorMode/Source/PLUGIN_NAME/Private/PLUGIN_NAMEModule.cpp
--- a/Engine/Plugins/Editor/PluginBrowser/Templates/EditorMode/Source/PLUGIN_NAME/Private/PLUGIN_NAMEModule.cpp
+++ b/Engine/Plugins/Editor/PluginBrowser/Templates/EditorMode/Source/PLUGIN_NAME/Private/PLUGIN_NAMEModule.cpp
@@ -22,4 +22,4 @@ void FPLUGIN_NAMEModule::ShutdownModule()
 
 #undef LOCTEXT_NAMESPACE
 
-IMPLEMENT_MODULE(FPLUGIN_NAMEModule, PLUGIN_NAMEEditorMode)
\ No newline at end of file
+IMPLEMENT_MODULE(FPLUGIN_NAMEModule, PLUGIN_NAME)
\ No newline at end of file
diff --git a/Engine/Plugins/Editor/PluginBrowser/Templates/EditorMode/Source/PLUGIN_NAME/Private/Tools/PLUGIN_NAMEInteractiveTool.cpp b/Engine/Plugins/Editor/PluginBrowser/Templates/EditorMode/Source/PLUGIN_NAME/Private/Tools/PLUGIN_NAMEInteractiveTool.cpp
--- a/Engine/Plugins/Editor/PluginBrowser/Templates/EditorMode/Source/PLUGIN_NAME/Private/Tools/PLUGIN_NAMEInteractiveTool.cpp
+++ b/Engine/Plugins/Editor/PluginBrowser/Templates/EditorMode/Source/PLUGIN_NAME/Private/Tools/PLUGIN_NAMEInteractiveTool.cpp
@@ -8,6 +8,7 @@
 // for raycast into World
 #include "CollisionQueryParams.h"
 #include "Engine/World.h"
+#include "Engine/HitResult.h"
 
 #include "SceneManagement.h"
 
diff --git a/Engine/Plugins/Editor/PluginBrowser/Templates/EditorMode/Source/PLUGIN_NAME/Private/Tools/PLUGIN_NAMESimpleTool.cpp b/Engine/Plugins/Editor/PluginBrowser/Templates/EditorMode/Source/PLUGIN_NAME/Private/Tools/PLUGIN_NAMESimpleTool.cpp
--- a/Engine/Plugins/Editor/PluginBrowser/Templates/EditorMode/Source/PLUGIN_NAME/Private/Tools/PLUGIN_NAMESimpleTool.cpp
+++ b/Engine/Plugins/Editor/PluginBrowser/Templates/EditorMode/Source/PLUGIN_NAME/Private/Tools/PLUGIN_NAMESimpleTool.cpp
@@ -5,6 +5,8 @@
 #include "ToolBuilderUtil.h"
 #include "CollisionQueryParams.h"
 #include "Engine/World.h"
+#include "Engine/HitResult.h"
+#include "Misc/MessageDialog.h"
 
 // localization namespace
 #define LOCTEXT_NAMESPACE "PLUGIN_NAMESimpleTool"
